Enable VirtualCamera plugin

#Synced-CL 51828360

diff --git a/Templates/TP_ME_CineBP/TP_ME_CineBP.uproject b/Templates/TP_ME_CineBP/TP_ME_CineBP.uproject
--- a/Templates/TP_ME_CineBP/TP_ME_CineBP.uproject
+++ b/Templates/TP_ME_CineBP/TP_ME_CineBP.uproject
@@ -27,6 +27,10 @@
 		{
 			"Name": "OpenColorIO",
 			"Enabled": true
+		},
+		{
+			"Name": "VirtualCamera",
+			"Enabled": true
 		}
 	]
-}
+}
\ No newline at end of file
