Added morph editing and media related plugins. Updated a few render sttings. Updated description texts in the default map.

#rnx
#virtualized
#Synced-CL 51828360

diff --git a/Templates/TP_ME_CineBP/Config/DefaultEngine.ini b/Templates/TP_ME_CineBP/Config/DefaultEngine.ini
--- a/Templates/TP_ME_CineBP/Config/DefaultEngine.ini
+++ b/Templates/TP_ME_CineBP/Config/DefaultEngine.ini
@@ -41,6 +41,8 @@ r.DefaultFeature.LensFlare=True
 r.PostProcessing.PropagateAlpha=True
 r.Deferred.SupportPrimitiveAlphaHoldout=True
 
+r.RayTracing.ResidentGeometryMemoryPoolSizeInMB=800
+
 [/Script/EngineSettings.GameMapsSettings]
 EditorStartupMap=/Game/CinematicTemplate/Maps/Main.Main
 
@@ -54,3 +56,10 @@ DefaultGraphicsRHI=DefaultGraphicsRHI_DX12
 [/Script/LinuxTargetPlatform.LinuxTargetSettings]
 -TargetedRHIs=SF_VULKAN_SM5
 +TargetedRHIs=SF_VULKAN_SM6
+
+[/Script/ElectraProtronFactory.ElectraProtronFactorySettings]
+bPreferProtronInEditor=True
+bPreferProtronInGame=True
+
+[/Script/Engine.VirtualTexturePoolConfig]
+DefaultSizeInMegabyte=256
diff --git a/Templates/TP_ME_CineBP/TP_ME_CineBP.uproject b/Templates/TP_ME_CineBP/TP_ME_CineBP.uproject
--- a/Templates/TP_ME_CineBP/TP_ME_CineBP.uproject
+++ b/Templates/TP_ME_CineBP/TP_ME_CineBP.uproject
@@ -1,6 +1,6 @@
 {
 	"FileVersion": 3,
-	"EngineAssociation": "",
+	"EngineAssociation": "{51E89314-4163-BEDF-50DA-2DBC6D3362DA}",
 	"Category": "",
 	"Description": "",
 	"Plugins": [
@@ -28,9 +28,29 @@
 			"Name": "OpenColorIO",
 			"Enabled": true
 		},
+		{
+			"Name": "SkeletalMeshMorphTargetEditingTools",
+			"Enabled": true
+		},
 		{
 			"Name": "VirtualCamera",
 			"Enabled": true
+		},
+		{
+			"Name": "ElectraPlayer",
+			"Enabled": true
+		},
+		{
+			"Name": "AppleProResDecoderElectra",
+			"Enabled": true,
+			"SupportedTargetPlatforms": [
+				"Win64",
+				"Mac"
+			]
+		},
+		{
+			"Name": "MediaViewer",
+			"Enabled": true
 		}
 	]
 }
\ No newline at end of file
