Enable RelativeIKOp and CinematicSequenceNavigator plugins

#rnx
#Synced-CL 52066367

diff --git a/Templates/TP_ME_CineBP/TP_ME_CineBP.uproject b/Templates/TP_ME_CineBP/TP_ME_CineBP.uproject
--- a/Templates/TP_ME_CineBP/TP_ME_CineBP.uproject
+++ b/Templates/TP_ME_CineBP/TP_ME_CineBP.uproject
@@ -51,6 +51,19 @@
 		{
 			"Name": "MediaViewer",
 			"Enabled": true
+		},
+		{
+			"Name": "RelativeIKOp",
+			"Enabled": true
+		},
+		{
+			"Name": "CinematicSequenceNavigator",
+			"Enabled": true,
+			"SupportedTargetPlatforms": [
+				"Win64",
+				"Linux",
+				"Mac"
+			]
 		}
 	]
 }
\ No newline at end of file
