Enable more plugins. Add an example LS to a level

#rnx
#virtualized
#Synced-CL 52403392

diff --git a/Templates/TP_ME_CineBP/TP_ME_CineBP.uproject b/Templates/TP_ME_CineBP/TP_ME_CineBP.uproject
--- a/Templates/TP_ME_CineBP/TP_ME_CineBP.uproject
+++ b/Templates/TP_ME_CineBP/TP_ME_CineBP.uproject
@@ -64,6 +64,53 @@
 				"Linux",
 				"Mac"
 			]
+		},
+		{
+			"Name": "ColorCorrectRegions",
+			"Enabled": true
+		},
+		{
+			"Name": "NiagaraMRQ",
+			"Enabled": true
+		},
+		{
+			"Name": "AppleProResMedia",
+			"Enabled": true,
+			"SupportedTargetPlatforms": [
+				"Win64"
+			]
+		},
+		{
+			"Name": "AvidDNxHDDecoderElectra",
+			"Enabled": true,
+			"SupportedTargetPlatforms": [
+				"Win64"
+			]
+		},
+		{
+			"Name": "AvidDNxMedia",
+			"Enabled": true
+		},
+		{
+			"Name": "MetaHuman",
+			"Enabled": true,
+			"SupportedTargetPlatforms": [
+				"Win64",
+				"Linux",
+				"Mac"
+			]
+		},
+		{
+			"Name": "AlembicHairImporter",
+			"Enabled": true
+		},
+		{
+			"Name": "AIAssistant",
+			"Enabled": true
+		},
+		{
+			"Name": "AccumulationDOF",
+			"Enabled": true
 		}
 	]
 }
\ No newline at end of file
